Fix up calls to OpenMsgDlgInt to use FMessageDialog::Open instead

diff --git a/Engine/Plugins/Editor/GameplayTagsEditor/Source/GameplayTagsEditor/Private/SGameplayTagWidget.cpp b/Engine/Plugins/Editor/GameplayTagsEditor/Source/GameplayTagsEditor/Private/SGameplayTagWidget.cpp
--- a/Engine/Plugins/Editor/GameplayTagsEditor/Source/GameplayTagsEditor/Private/SGameplayTagWidget.cpp
+++ b/Engine/Plugins/Editor/GameplayTagsEditor/Source/GameplayTagsEditor/Private/SGameplayTagWidget.cpp
@@ -8,7 +8,7 @@
 #include "Widgets/Images/SImage.h"
 #include "EditorStyleSet.h"
 #include "Widgets/SWindow.h"
-#include "Dialogs/Dialogs.h"
+#include "Misc/MessageDialog.h"
 #include "GameplayTagsModule.h"
 #include "ScopedTransaction.h"
 #include "Textures/SlateIcon.h"
@@ -1070,7 +1070,8 @@ void SGameplayTagWidget::VerifyAssetTagValidity()
 				FFormatNamedArguments Arguments;
 				Arguments.Add(TEXT("Objects"), FText::FromString( InvalidTagNames ));
 				FText DialogText = FText::Format( LOCTEXT("GameplayTagWidget_InvalidTags", "Invalid Tags that have been removed: \n\n{Objects}"), Arguments );
-				OpenMsgDlgInt( EAppMsgType::Ok, DialogText, LOCTEXT("GameplayTagWidget_Warning", "Warning") );
+				FText DialogTitle = LOCTEXT("GameplayTagWidget_Warning", "Warning");
+				FMessageDialog::Open( EAppMsgType::Ok, DialogText, &DialogTitle );
 			}
 		}
 	}
diff --git a/Engine/Plugins/Experimental/DataPrepEditor/Source/DataPrepEditor/Private/DataPrepEditor.cpp b/Engine/Plugins/Experimental/DataPrepEditor/Source/DataPrepEditor/Private/DataPrepEditor.cpp
--- a/Engine/Plugins/Experimental/DataPrepEditor/Source/DataPrepEditor/Private/DataPrepEditor.cpp
+++ b/Engine/Plugins/Experimental/DataPrepEditor/Source/DataPrepEditor/Private/DataPrepEditor.cpp
@@ -24,7 +24,7 @@
 #include "BlueprintNodeSpawner.h"
 #include "DesktopPlatformModule.h"
 #include "Dialogs/DlgPickPath.h"
-#include "Dialogs/Dialogs.h"
+#include "Misc/MessageDialog.h"
 #include "Editor.h"
 #include "EditorDirectories.h"
 #include "EditorStyleSet.h"
@@ -902,7 +902,7 @@ void FDataprepEditor::OnCommitWorld()
 			const FText Title( LOCTEXT( "DataprepEditor_ProceedWithCommit", "Proceed with commit" ) );
 			const FText Message( LOCTEXT( "DataprepEditor_ConfirmCommitPipelineNotExecuted", "The action pipeline has not been executed.\nDo you want to proceeed with the commit anyway?" ) );
 
-			if( OpenMsgDlgInt( EAppMsgType::YesNo, Message, Title ) == EAppReturnType::No )
+			if( FMessageDialog::Open( EAppMsgType::YesNo, Message, &Title ) == EAppReturnType::No )
 			{
 				return;
 			}
@@ -914,7 +914,7 @@ void FDataprepEditor::OnCommitWorld()
 		const FText Title( LOCTEXT( "DataprepEditor_ProceedWithCommit", "Proceed with commit" ) );
 		const FText Message( LOCTEXT( "DataprepEditor_ConfirmCommitPipelineChanged", "The action pipeline has changed since last execution.\nDo you want to proceeed with the commit anyway?" ) );
 
-		if( OpenMsgDlgInt( EAppMsgType::YesNo, Message, Title ) == EAppReturnType::No )
+		if( FMessageDialog::Open( EAppMsgType::YesNo, Message, &Title ) == EAppReturnType::No )
 		{
 			return;
 		}
diff --git a/Engine/Source/Developer/AssetTools/Private/AssetTypeActions/AssetTypeActions_AnimCurveCompressionSettings.cpp b/Engine/Source/Developer/AssetTools/Private/AssetTypeActions/AssetTypeActions_AnimCurveCompressionSettings.cpp
--- a/Engine/Source/Developer/AssetTools/Private/AssetTypeActions/AssetTypeActions_AnimCurveCompressionSettings.cpp
+++ b/Engine/Source/Developer/AssetTools/Private/AssetTypeActions/AssetTypeActions_AnimCurveCompressionSettings.cpp
@@ -2,7 +2,7 @@
 
 #include "AssetTypeActions/AssetTypeActions_AnimCurveCompressionSettings.h"
 #include "Animation/AnimSequence.h"
-#include "Dialogs/Dialogs.h"
+#include "Misc/MessageDialog.h"
 #include "EditorStyleSet.h"
 #include "ToolMenus.h"
 #include "Misc/ScopedSlowTask.h"
@@ -93,7 +93,8 @@ void FAssetTypeActions_AnimCurveCompressionSettings::ExecuteCompression(TWeakObj
 	FFormatNamedArguments Arguments;
 	Arguments.Add(TEXT("NumAnimSequences"), FText::AsNumber(AnimSeqsToRecompress.Num()));
 	FText DialogText = FText::Format(LOCTEXT("AnimCurveCompressionSettings_CompressWarningText", "{NumAnimSequences} animation sequences are about to compress."), Arguments);
-	const EAppReturnType::Type DlgResult = OpenMsgDlgInt(EAppMsgType::OkCancel, DialogText, LOCTEXT("AnimCurveCompressionSettings_CompressWarning", "Warning"));
+	FText DialogTitle = LOCTEXT("AnimCurveCompressionSettings_CompressWarning", "Warning");
+	const EAppReturnType::Type DlgResult = FMessageDialog::Open(EAppMsgType::OkCancel, DialogText, &DialogTitle);
 	if (DlgResult != EAppReturnType::Ok)
 	{
 		return;
diff --git a/Engine/Source/Developer/AssetTools/Private/AssetTypeActions/AssetTypeActions_DataTable.cpp b/Engine/Source/Developer/AssetTools/Private/AssetTypeActions/AssetTypeActions_DataTable.cpp
--- a/Engine/Source/Developer/AssetTools/Private/AssetTypeActions/AssetTypeActions_DataTable.cpp
+++ b/Engine/Source/Developer/AssetTools/Private/AssetTypeActions/AssetTypeActions_DataTable.cpp
@@ -4,7 +4,7 @@
 #include "ToolMenus.h"
 #include "Misc/FileHelper.h"
 #include "EditorFramework/AssetImportData.h"
-#include "Dialogs/Dialogs.h"
+#include "Misc/MessageDialog.h"
 #include "Framework/Application/SlateApplication.h"
 
 #include "Editor/DataTableEditor/Public/DataTableEditorModule.h"
@@ -165,10 +165,11 @@ void FAssetTypeActions_DataTable::OpenAssetEditor( const TArray<UObject*>& InObj
 			DataTablesListText.AppendLineFormat(LOCTEXT("DataTable_MissingRowStructListEntry", "* {0} (Row Structure: {1})"), FText::FromString(Table->GetName()), FText::FromName(ResolvedRowStructName));
 		}
 
-		const EAppReturnType::Type DlgResult = OpenMsgDlgInt(
+		FText Title = LOCTEXT("DataTable_MissingRowStructTitle", "Continue?");
+		const EAppReturnType::Type DlgResult = FMessageDialog::Open(
 			EAppMsgType::YesNoCancel, 
 			FText::Format(LOCTEXT("DataTable_MissingRowStructMsg", "The following Data Tables are missing their row structure and will not be editable.\n\n{0}\n\nDo you want to open these data tables?"), DataTablesListText.ToText()), 
-			LOCTEXT("DataTable_MissingRowStructTitle", "Continue?")
+			&Title
 			);
 
 		switch(DlgResult)
diff --git a/Engine/Source/Developer/MaterialBaking/Private/MaterialOptionsWindow.cpp b/Engine/Source/Developer/MaterialBaking/Private/MaterialOptionsWindow.cpp
--- a/Engine/Source/Developer/MaterialBaking/Private/MaterialOptionsWindow.cpp
+++ b/Engine/Source/Developer/MaterialBaking/Private/MaterialOptionsWindow.cpp
@@ -11,7 +11,7 @@
 #include "MaterialBakingStructures.h"
 #include "MaterialOptions.h"
 #include "MaterialOptionsCustomization.h"
-#include "Dialogs/Dialogs.h"
+#include "Misc/MessageDialog.h"
 
 #define LOCTEXT_NAMESPACE "SMaterialOptions"
 
@@ -82,7 +82,8 @@ FReply SMaterialOptions::OnConfirm()
 	// Ensure the user has selected at least one LOD index
 	if (GetMutableDefault<UMaterialOptions>()->LODIndices.Num() == 0)
 	{
-		OpenMsgDlgInt(EAppMsgType::Ok, LOCTEXT("MaterialBake_SelectLODError", "Ensure that atleast one LOD index is selected."), LOCTEXT("MaterialBake_SelectLODErrorTitle", "Invalid options"));
+		FText Title = LOCTEXT("MaterialBake_SelectLODErrorTitle", "Invalid options");
+		FMessageDialog::Open(EAppMsgType::Ok, LOCTEXT("MaterialBake_SelectLODError", "Ensure that atleast one LOD index is selected."), &Title);
 	}
 	else
 	{
diff --git a/Engine/Source/Editor/AnimationModifiers/Private/AnimationModifier.cpp b/Engine/Source/Editor/AnimationModifiers/Private/AnimationModifier.cpp
--- a/Engine/Source/Editor/AnimationModifiers/Private/AnimationModifier.cpp
+++ b/Engine/Source/Editor/AnimationModifiers/Private/AnimationModifier.cpp
@@ -7,7 +7,7 @@
 #include "Editor/Transactor.h"
 #include "UObject/UObjectIterator.h"
 
-#include "Dialogs/Dialogs.h"
+#include "Misc/MessageDialog.h"
 #include "Editor/Transactor.h"
 #include "UObject/UObjectIterator.h"
 
@@ -55,7 +55,8 @@ void UAnimationModifier::ApplyToAnimationSequence(class UAnimSequence* InAnimati
 
 		EAppMsgType::Type MessageType = OutputLog.ContainsErrors() ? EAppMsgType::Ok : EAppMsgType::YesNo;
 		const FText& MessageFormat = OutputLog.ContainsErrors() ? ErrorMessageFormat : WarningMessageFormat;
-		bShouldRevert = (OpenMsgDlgInt(MessageType, FText::FormatOrdered(MessageFormat, FText::FromString(OutputLog)), FText::FromString("Modifier has Generated Warnings/Errors")) != EAppReturnType::Yes);
+		const FText MessageTitle = FText::FromString("Modifier has Generated Warnings/Errors");
+		bShouldRevert = (FMessageDialog::Open(MessageType, FText::FormatOrdered(MessageFormat, FText::FromString(OutputLog)), &MessageTitle) != EAppReturnType::Yes);
 	}
 
 	// Revert changes if necessary, otherwise post edit and refresh animation data
diff --git a/Engine/Source/Editor/AnimationModifiers/Private/SAnimationModifierContentBrowserWindow.cpp b/Engine/Source/Editor/AnimationModifiers/Private/SAnimationModifierContentBrowserWindow.cpp
--- a/Engine/Source/Editor/AnimationModifiers/Private/SAnimationModifierContentBrowserWindow.cpp
+++ b/Engine/Source/Editor/AnimationModifiers/Private/SAnimationModifierContentBrowserWindow.cpp
@@ -15,7 +15,7 @@
 
 #include "Animation/AnimSequence.h"
 #include "AnimationModifiersAssetUserData.h"
-#include "Dialogs/Dialogs.h"
+#include "Misc/MessageDialog.h"
 #include "ScopedTransaction.h"
 
 #define LOCTEXT_NAMESPACE "AnimationModifierContentBrowserWindow"
@@ -211,9 +211,15 @@ FReply SAnimationModifierContentBrowserWindow::OnApply()
 		for (UAnimationModifiersAssetUserData* UserData : AssetUserData)
 		{
 			const bool bAlreadyContainsModifier = UserData->GetAnimationModifierInstances().ContainsByPredicate([Modifier](UAnimationModifier* TestModifier) { return Modifier->GetClass() == TestModifier->GetClass(); });
-			static const FText MessageFormat = LOCTEXT("AnimationModifierWindow_AlreadyContainsModifierDialogText", "{0} already contains Animation Modifier {1}, are you sure you want to add another instance?");
 
-			const bool bUserInputResult = bAlreadyContainsModifier ? OpenMsgDlgInt(EAppMsgType::YesNo, FText::FormatOrdered(MessageFormat, FText::FromString(UserData->GetOuter()->GetName()), FText::FromString(Modifier->GetClass()->GetName())), LOCTEXT("AnimationModifierWindow_AlreadyContainsModifierTitle", "Already contains Animation Modifier!")) == EAppReturnType::Yes : true;
+			bool bUserInputResult = true;
+			if (bAlreadyContainsModifier)
+			{
+				FText MessageFormat = LOCTEXT("AnimationModifierWindow_AlreadyContainsModifierDialogText", "{0} already contains Animation Modifier {1}, are you sure you want to add another instance?");
+				FText MessageTitle = LOCTEXT("AnimationModifierWindow_AlreadyContainsModifierTitle", "Already contains Animation Modifier!");
+				bUserInputResult = FMessageDialog::Open(EAppMsgType::YesNo, FText::FormatOrdered(MessageFormat, FText::FromString(UserData->GetOuter()->GetName()), FText::FromString(Modifier->GetClass()->GetName())), &MessageTitle) == EAppReturnType::Yes;
+			}
+			
 			bCloseWindow = bUserInputResult;
 			if (!bAlreadyContainsModifier || bUserInputResult)
 			{
diff --git a/Engine/Source/Editor/AnimationModifiers/Private/SAnimationModifiersTab.cpp b/Engine/Source/Editor/AnimationModifiers/Private/SAnimationModifiersTab.cpp
--- a/Engine/Source/Editor/AnimationModifiers/Private/SAnimationModifiersTab.cpp
+++ b/Engine/Source/Editor/AnimationModifiers/Private/SAnimationModifiersTab.cpp
@@ -15,7 +15,7 @@
 #include "ClassViewerFilter.h"
 #include "PropertyEditorModule.h"
 #include "IDetailsView.h"
-#include "Dialogs/Dialogs.h"
+#include "Misc/MessageDialog.h"
 #include "Editor.h"
 #include "ScopedTransaction.h"
 
@@ -319,7 +319,8 @@ void SAnimationModifiersTab::OnRevertModifier(const TArray<TWeakObjectPtr<UAnima
 
 void SAnimationModifiersTab::OnRemoveModifier(const TArray<TWeakObjectPtr<UAnimationModifier>>& Instances)
 {
-	const bool bShouldRevert = OpenMsgDlgInt(EAppMsgType::YesNo, LOCTEXT("RemoveAndRevertPopupText", "Should the Modifiers be reverted before removing them?"), FText::FromString("Revert before Removing")) == EAppReturnType::Yes;
+	const FText Title = FText::FromString("Revert before Removing");
+	const bool bShouldRevert = FMessageDialog::Open(EAppMsgType::YesNo, LOCTEXT("RemoveAndRevertPopupText", "Should the Modifiers be reverted before removing them?"), &Title) == EAppReturnType::Yes;
 
 	FScopedTransaction Transaction(LOCTEXT("RemoveModifiersTransaction", "Removing Animation Modifier(s)"));	
 	AssetUserData->Modify();
@@ -417,7 +418,8 @@ void SAnimationModifiersTab::ApplyModifiers(const TArray<UAnimationModifier*>& M
 	else if (Skeleton != nullptr)
 	{
 		// Double check with the user for applying all modifiers to referenced animation sequences for the skeleton
-		bApply = OpenMsgDlgInt(EAppMsgType::YesNo, LOCTEXT("ApplyingSkeletonModifierPopupText", "Are you sure you want to apply the modifiers to all animation sequences referenced by the current skeleton?"), FText::FromString("Are you sure?")) == EAppReturnType::Yes;
+		const FText Title = FText::FromString("Are you sure?");
+		bApply = FMessageDialog::Open(EAppMsgType::YesNo, LOCTEXT("ApplyingSkeletonModifierPopupText", "Are you sure you want to apply the modifiers to all animation sequences referenced by the current skeleton?"), &Title) == EAppReturnType::Yes;
 		
 		if (bApply)
 		{
@@ -457,7 +459,8 @@ void SAnimationModifiersTab::RevertModifiers(const TArray<UAnimationModifier*>&
 	else if (Skeleton != nullptr)
 	{
 		// Double check with the user for reverting all modifiers from referenced animation sequences for the skeleton
-		bRevert = OpenMsgDlgInt(EAppMsgType::YesNo, LOCTEXT("RevertingSkeletonModifierPopupText", "Are you sure you want to revert the modifiers from all animation sequences referenced by the current skeleton?"), FText::FromString("Are you sure?")) == EAppReturnType::Yes;
+		const FText Title = FText::FromString("Are you sure?");
+		bRevert = FMessageDialog::Open(EAppMsgType::YesNo, LOCTEXT("RevertingSkeletonModifierPopupText", "Are you sure you want to revert the modifiers from all animation sequences referenced by the current skeleton?"), &Title) == EAppReturnType::Yes;
 
 		if ( bRevert)
 		{
diff --git a/Engine/Source/Editor/FontEditor/Private/FontEditor.cpp b/Engine/Source/Editor/FontEditor/Private/FontEditor.cpp
--- a/Engine/Source/Editor/FontEditor/Private/FontEditor.cpp
+++ b/Engine/Source/Editor/FontEditor/Private/FontEditor.cpp
@@ -431,10 +431,11 @@ void FFontEditor::NotifyPostChange( const FPropertyChangedEvent& PropertyChanged
 	if(PropertyChangedEvent.Property && PropertyChangedEvent.Property->GetFName() == FontCacheTypePropertyName)
 	{
 		// Show a warning message, as what we're about to do will destroy any existing data in this font object
-		const EAppReturnType::Type DlgResult = OpenMsgDlgInt(
+		const FText Title = LOCTEXT("ChangeCacheTypeWarningTitle", "Really change the font cache type?");
+		const EAppReturnType::Type DlgResult = FMessageDialog::Open(
 			EAppMsgType::YesNo, 
 			LOCTEXT("ChangeCacheTypeWarningMsg", "Changing the cache type will cause this font to be reinitialized (discarding any existing data).\n\nAre you sure you want to proceed?"), 
-			LOCTEXT("ChangeCacheTypeWarningTitle", "Really change the font cache type?")
+			&Title
 			);
 
 		bool bSuccessfullyChangedCacheType = false;
diff --git a/Engine/Source/Editor/InputBindingEditor/Private/InputBindingEditorModule.cpp b/Engine/Source/Editor/InputBindingEditor/Private/InputBindingEditorModule.cpp
--- a/Engine/Source/Editor/InputBindingEditor/Private/InputBindingEditorModule.cpp
+++ b/Engine/Source/Editor/InputBindingEditor/Private/InputBindingEditorModule.cpp
@@ -8,7 +8,7 @@
 #include "Widgets/SInputBindingEditorPanel.h"
 #include "UnrealEdMisc.h"
 #include "Logging/MessageLog.h"
-#include "Dialogs/Dialogs.h"
+#include "Misc/MessageDialog.h"
 #include "IDetailCustomization.h"
 #include "ISettingsModule.h"
 #include "PropertyEditorModule.h"
@@ -95,7 +95,7 @@ private:
 	// Show a warning that the editor will require a restart and return its result
 	EAppReturnType::Type ShowRestartWarning(const FText& Title) const
 	{
-		return OpenMsgDlgInt(EAppMsgType::OkCancel, LOCTEXT("ActionRestartMsg", "Imported settings won't be applied until the editor is restarted. Do you wish to restart now (you will be prompted to save any changes)?"), Title);
+		return FMessageDialog::Open(EAppMsgType::OkCancel, LOCTEXT("ActionRestartMsg", "Imported settings won't be applied until the editor is restarted. Do you wish to restart now (you will be prompted to save any changes)?"), &Title);
 	}
 
 	// Backup a file
diff --git a/Engine/Source/Editor/MainFrame/Private/Menus/LayoutsMenu.cpp b/Engine/Source/Editor/MainFrame/Private/Menus/LayoutsMenu.cpp
--- a/Engine/Source/Editor/MainFrame/Private/Menus/LayoutsMenu.cpp
+++ b/Engine/Source/Editor/MainFrame/Private/Menus/LayoutsMenu.cpp
@@ -24,7 +24,7 @@
 #include "ToolMenus.h"
 // Editor
 #include "Classes/EditorStyleSettings.h"
-#include "Dialogs/Dialogs.h"
+#include "Misc/MessageDialog.h"
 #include "Editor/EditorPerProjectUserSettings.h"
 #include "Frame/MainFrameActions.h"
 #include "LevelViewportActions.h"
@@ -109,7 +109,7 @@ bool TrySaveLayoutOrWarnInternal(const FString& InSourceFilePath, const FString&
 							"If you are not sure, select \"Yes\" if you are exporting the layout configuration without making any changes, or \"No\" if you have made or plan to make changes to the layout.\n\n"
 						),
 						LayoutNameSource, LayoutDescriptionSource);
-					const int AppReturnType = OpenMsgDlgInt(EAppMsgType::YesNoCancel, TextBody, TextTitle);
+					const int AppReturnType = FMessageDialog::Open(EAppMsgType::YesNoCancel, TextBody, &TextTitle);
 					// Handle user answers
 					if (AppReturnType == EAppReturnType::Yes)
 					{
@@ -171,7 +171,7 @@ bool TrySaveLayoutOrWarnInternal(const FString& InSourceFilePath, const FString&
 		EditorErrors.Notify(LOCTEXT("LoadUnsuccessful_Title", "Load Unsuccessful!"));
 		// Show reason
 		const FText TextTitle = LOCTEXT("UnsuccessfulCopyHeader", "Unsuccessful copy!");
-		OpenMsgDlgInt(EAppMsgType::Ok, TextBody, TextTitle);
+		FMessageDialog::Open(EAppMsgType::Ok, TextBody, &TextTitle);
 		// Return
 		return false;
 	}
@@ -571,7 +571,7 @@ void FLayoutsMenuLoad::ImportLayout()
 					Arguments.Add(TEXT("FileName"), FText::FromString(FPaths::ConvertRelativePathToFull(LayoutFilePath)));
 					const FText TextBody = FText::Format(LOCTEXT("UnsuccessfulImportBody", "Unsuccessful import, {FileName} is not a layout configuration file!"), Arguments);
 					const FText TextTitle = LOCTEXT("UnsuccessfulImportHeader", "Unsuccessful Import!");
-					OpenMsgDlgInt(EAppMsgType::Ok, TextBody, TextTitle);
+					FMessageDialog::Open(EAppMsgType::Ok, TextBody, &TextTitle);
 				}
 			}
 			// If PIE running, do not reload current layout
@@ -580,7 +580,7 @@ void FLayoutsMenuLoad::ImportLayout()
 				const FText TextBody = LOCTEXT("SuccessfulImportBody",
 					"The layout(s) were successfully imported into the \"User Layouts\" section. However, no layout has been loaded into your current Unreal Editor UI because PIE is currently running. In order to do so, you must stop PIE and then load the layout from the \"User Layouts\" section.");
 				const FText TextTitle = LOCTEXT("UnsuccessfulImportHeader", "Successful Import!");
-				OpenMsgDlgInt(EAppMsgType::Ok, TextBody, TextTitle);
+				FMessageDialog::Open(EAppMsgType::Ok, TextBody, &TextTitle);
 				return;
 			}
 			// Replace current layout with first one
@@ -644,7 +644,8 @@ void SaveExportLayoutCommon(const FString& InDefaultDirectory, const bool bMustB
 				else
 				{
 					// Warn the user that the file will not be copied in there
-					OpenMsgDlgInt(
+					const FText Title = LOCTEXT("SaveAsFailedMsg_Title", "Save As Failed");
+					FMessageDialog::Open(
 						EAppMsgType::Ok,
 						FText::Format(
 							LOCTEXT("SaveAsFailedMsg",
@@ -652,7 +653,7 @@ void SaveExportLayoutCommon(const FString& InDefaultDirectory, const bool bMustB
 								"If you simply wish to export a copy of the current configuration in {1} (e.g., to later copy it into a different machine), you could use the \"Export Layout...\""
 								" functionality. However, Unreal would not be able to load it until you import it with \"Import Layout...\"."),
 							FText::FromString(DefaultDirectoryAbsolute), FText::FromString(LayoutFilePathAbsolute)),
-						LOCTEXT("SaveAsFailedMsg_Title", "Save As Failed"));
+						&Title);
 				}
 			}
 		}
@@ -716,7 +717,7 @@ void FLayoutsMenuSave::OverrideUserLayout(const int32 InLayoutIndex)
 	const FText TextFileNameToRemove = FText::FromString(FPaths::GetBaseFilename(UserLayoutIniFileNames[InLayoutIndex]));
 	const FText TextBody = FText::Format(LOCTEXT("ActionOverrideLayoutMsg", "Are you sure you want to permanently override the layout profile \"{0}\" with the current layout profile? This action cannot be undone."), TextFileNameToRemove);
 	const FText TextTitle = FText::Format(LOCTEXT("OverrideUILayout_Title", "Override UI Layout \"{0}\""), TextFileNameToRemove);
-	if (EAppReturnType::Ok != OpenMsgDlgInt(EAppMsgType::OkCancel, TextBody, TextTitle))
+	if (EAppReturnType::Ok != FMessageDialog::Open(EAppMsgType::OkCancel, TextBody, &TextTitle))
 	{
 		return;
 	}
@@ -818,7 +819,7 @@ void FLayoutsMenuRemove::RemoveUserLayout(const int32 InLayoutIndex)
 	const FText TextFileNameToRemove = FText::FromString(FPaths::GetBaseFilename(UserLayoutIniFileNames[InLayoutIndex]));
 	const FText TextBody = FText::Format(LOCTEXT("ActionRemoveMsg", "Are you sure you want to permanently delete the layout profile \"{0}\"? This action cannot be undone."), TextFileNameToRemove);
 	const FText TextTitle = FText::Format(LOCTEXT("RemoveUILayout_Title", "Remove UI Layout \"{0}\""), TextFileNameToRemove);
-	if (EAppReturnType::Ok != OpenMsgDlgInt(EAppMsgType::OkCancel, TextBody, TextTitle))
+	if (EAppReturnType::Ok != FMessageDialog::Open(EAppMsgType::OkCancel, TextBody, &TextTitle))
 	{
 		return;
 	}
@@ -857,7 +858,7 @@ void FLayoutsMenuRemove::RemoveUserLayouts()
 		// Are you sure you want to do this?
 		const FText TextBody = FText::Format(LOCTEXT("ActionRemoveAllUserLayoutMsg", "Are you sure you want to permanently remove {0} layout {0}|plural(one=profile,other=profiles)? This action cannot be undone."), NumberUserLayoutFiles);
 		const FText TextTitle = LOCTEXT("RemoveAllUserLayouts_Title", "Remove All User-Created Layouts");
-		if (EAppReturnType::Ok != OpenMsgDlgInt(EAppMsgType::OkCancel, TextBody, TextTitle))
+		if (EAppReturnType::Ok != FMessageDialog::Open(EAppMsgType::OkCancel, TextBody, &TextTitle))
 		{
 			return;
 		}
@@ -881,7 +882,7 @@ void FLayoutsMenuRemove::RemoveUserLayouts()
 		// Show reason
 		const FText TextBody = LOCTEXT("UnsuccessfulRemoveLayoutBody", "There are no layout profile files created by the user, so none could be removed.");
 		const FText TextTitle = LOCTEXT("UnsuccessfulRemoveLayoutHeader", "Unsuccessful Remove All User Layouts!");
-		OpenMsgDlgInt(EAppMsgType::Ok, TextBody, TextTitle);
+		FMessageDialog::Open(EAppMsgType::Ok, TextBody, &TextTitle);
 	}
 }
 
diff --git a/Engine/Source/Editor/MergeActors/Private/MeshInstancingTool/MeshInstancingTool.cpp b/Engine/Source/Editor/MergeActors/Private/MeshInstancingTool/MeshInstancingTool.cpp
--- a/Engine/Source/Editor/MergeActors/Private/MeshInstancingTool/MeshInstancingTool.cpp
+++ b/Engine/Source/Editor/MergeActors/Private/MeshInstancingTool/MeshInstancingTool.cpp
@@ -12,7 +12,7 @@
 #include "Engine/StaticMeshActor.h"
 #include "Engine/Selection.h"
 #include "Editor.h"
-#include "Dialogs/Dialogs.h"
+#include "Misc/MessageDialog.h"
 #include "MeshUtilities.h"
 #include "MeshInstancingTool/SMeshInstancingDialog.h"
 #include "IContentBrowserSingleton.h"
@@ -73,7 +73,8 @@ bool FMeshInstancingTool::RunMerge(const FString& PackageName)
 	if (UniqueLevels.Num() > 1)
 	{
 		FText Message = NSLOCTEXT("UnrealEd", "FailedToInstanceActorsSublevels_Msg", "The selected actors should be in the same level");
-		OpenMsgDlgInt(EAppMsgType::Ok, Message, NSLOCTEXT("UnrealEd", "FailedToInstanceActors_Title", "Unable to replace actors with instanced meshes"));
+		FText Title = NSLOCTEXT("UnrealEd", "FailedToInstanceActors_Title", "Unable to replace actors with instanced meshes");
+		FMessageDialog::Open(EAppMsgType::Ok, Message, &Title);
 		return false;
 	}
 
diff --git a/Engine/Source/Editor/MergeActors/Private/MeshMergingTool/MeshMergingTool.cpp b/Engine/Source/Editor/MergeActors/Private/MeshMergingTool/MeshMergingTool.cpp
--- a/Engine/Source/Editor/MergeActors/Private/MeshMergingTool/MeshMergingTool.cpp
+++ b/Engine/Source/Editor/MergeActors/Private/MeshMergingTool/MeshMergingTool.cpp
@@ -12,7 +12,7 @@
 #include "Engine/StaticMeshActor.h"
 #include "Engine/Selection.h"
 #include "Editor.h"
-#include "Dialogs/Dialogs.h"
+#include "Misc/MessageDialog.h"
 #include "MeshUtilities.h"
 #include "MeshMergingTool/SMeshMergingDialog.h"
 #include "IContentBrowserSingleton.h"
@@ -96,7 +96,8 @@ bool FMeshMergingTool::RunMerge(const FString& PackageName)
 	if (UniqueLevels.Num() > 1 && bReplaceSourceActors)
 	{
 		FText Message = NSLOCTEXT("UnrealEd", "FailedToMergeActorsSublevels_Msg", "The selected actors should be in the same level");
-		OpenMsgDlgInt(EAppMsgType::Ok, Message, NSLOCTEXT("UnrealEd", "FailedToMergeActors_Title", "Unable to merge actors"));
+		const FText Title = NSLOCTEXT("UnrealEd", "FailedToMergeActors_Title", "Unable to merge actors");
+		FMessageDialog::Open(EAppMsgType::Ok, Message, &Title);
 		return false;
 	}
 
diff --git a/Engine/Source/Editor/PListEditor/Private/SPlistEditor.cpp b/Engine/Source/Editor/PListEditor/Private/SPlistEditor.cpp
--- a/Engine/Source/Editor/PListEditor/Private/SPlistEditor.cpp
+++ b/Engine/Source/Editor/PListEditor/Private/SPlistEditor.cpp
@@ -8,7 +8,7 @@
 #include "Framework/MultiBox/MultiBoxBuilder.h"
 #include "Widgets/Input/SEditableText.h"
 #include "Widgets/Input/SButton.h"
-#include "Dialogs/Dialogs.h"
+#include "Misc/MessageDialog.h"
 
 #include "XmlFile.h"
 #include "PListNodeArray.h"
@@ -848,8 +848,10 @@ bool SPListEditorPanel::OpenFile(FString FilePath)
 		FText ErrorMessageFormatting = LOCTEXT("PListXMLLoadErrorFormatting", "Failed to Load PList File: {Filepath}\n\n{ErrorDetails}");
 		FText ErrorMessage = FText::Format( ErrorMessageFormatting, Arguments );
 
+		FText ErrorTitle = LOCTEXT("PListLoadFailDialogCaption", "Error");
+
 		// Show error message
-		OpenMsgDlgInt( EAppMsgType::Ok, ErrorMessage, LOCTEXT("PListLoadFailDialogCaption", "Error") );
+		FMessageDialog::Open( EAppMsgType::Ok, ErrorMessage, &ErrorTitle );
 		return false;
 	}
 	else
@@ -989,7 +991,8 @@ bool SPListEditorPanel::PromptDelete()
 	}
 
 	FText DeleteMessage = LOCTEXT("PListDeleteConfirmation", "Are you sure you want to remove the selected entries? (This action is irreversible!)");
-	EAppReturnType::Type ret = OpenMsgDlgInt( EAppMsgType::YesNoYesAll, DeleteMessage, LOCTEXT("PListDeleteConfirmationCaption", "Confirm Removal") );
+	FText DeleteTitle = LOCTEXT("PListDeleteConfirmationCaption", "Confirm Removal");
+	EAppReturnType::Type ret = FMessageDialog::Open( EAppMsgType::YesNoYesAll, DeleteMessage, &DeleteTitle );
 	if(ret == EAppReturnType::Yes)
 	{
 		return true;
@@ -1016,8 +1019,9 @@ bool SPListEditorPanel::PromptSave()
 	FFormatNamedArguments Arguments;
 	Arguments.Add(TEXT("FilePath"), FText::FromString( InOutLastPath ));
 	FText DialogText = FText::Format( LOCTEXT("PListCloseTabSaveTextFormatting", "Save {FilePath}?"), Arguments );
+	FText DialogTitle = LOCTEXT("PListCloseTabSaveCaption", "Save");
 
-	EAppReturnType::Type ret = OpenMsgDlgInt(EAppMsgType::YesNoCancel, DialogText, LOCTEXT("PListCloseTabSaveCaption", "Save") );
+	EAppReturnType::Type ret = FMessageDialog::Open(EAppMsgType::YesNoCancel, DialogText, &DialogTitle );
 	if(ret == EAppReturnType::Yes)
 	{
 		// Get the saving location if necessary (like on new files)
@@ -1055,8 +1059,9 @@ bool SPListEditorPanel::PromptSave()
 						Args.Add(TEXT("Filename"), FText::FromString(OutFilename));
 						FText OverwriteMessageFormatting = LOCTEXT("PListCloseTabOverwriteTextFormatting", "Overwrite existing file {Filename}?");
 						FText OverwriteDialogText = FText::Format(OverwriteMessageFormatting, Args);
+						FText OverwriteDialogTitle = LOCTEXT("PListWarningCaption", "Warning");
 
-						EAppReturnType::Type RetVal = OpenMsgDlgInt(EAppMsgType::YesNo, OverwriteDialogText, LOCTEXT("PListWarningCaption", "Warning"));
+						EAppReturnType::Type RetVal = FMessageDialog::Open(EAppMsgType::YesNo, OverwriteDialogText, &OverwriteDialogTitle);
 						if(RetVal != EAppReturnType::Yes)
 						{
 							// Said not to overwrite (or clicked x) so bail out
@@ -1085,7 +1090,8 @@ bool SPListEditorPanel::PromptSave()
 
 			// Display Message
 			FText ValidationFailMessage = LOCTEXT("PListNodeValidationFail", "Cannot save file: Not all plist entries have valid input");
-			OpenMsgDlgInt( EAppMsgType::Ok, ValidationFailMessage, LOCTEXT("PListWarningCaption", "Warning") );
+			FText Title = LOCTEXT("PListWarningCaption", "Warning");
+			FMessageDialog::Open( EAppMsgType::Ok, ValidationFailMessage, &Title );
 
 			// Cancel
 			return false;
@@ -1150,8 +1156,9 @@ FReply SPListEditorPanel::OnSaveClicked()
 		FFormatNamedArguments Arguments;
 		Arguments.Add(TEXT("FilePath"), FText::FromString( InOutLastPath ));
 		FText DialogText = FText::Format( LOCTEXT("PListOverwriteMessageFormatting", "Overwrite {FilePath}?"), Arguments );
+		FText DialogTitle = LOCTEXT("PListOverwriteCaption", "Warning");
 
-		EAppReturnType::Type ret = OpenMsgDlgInt( EAppMsgType::YesNo, DialogText, LOCTEXT("PListOverwriteCaption", "Warning") );
+		EAppReturnType::Type ret = FMessageDialog::Open( EAppMsgType::YesNo, DialogText, &DialogTitle );
 		if(ret != EAppReturnType::Yes)
 		{
 			return FReply::Handled();
@@ -1193,8 +1200,9 @@ FReply SPListEditorPanel::OnSaveClicked()
 					Arguments.Add(TEXT("Filename"), FText::FromString( OutFilename ));
 					FText OverwriteMessageFormatting = LOCTEXT("PListFileExistsMessageFormatting", "Overwrite existing file {Filename}?");
 					FText DialogText = FText::Format( OverwriteMessageFormatting, Arguments );
+					FText DialogTitle = LOCTEXT("PListWarningCaption", "Warning");
 
-					EAppReturnType::Type RetVal = OpenMsgDlgInt( EAppMsgType::YesNo, DialogText, LOCTEXT("PListWarningCaption", "Warning") );
+					EAppReturnType::Type RetVal = FMessageDialog::Open( EAppMsgType::YesNo, DialogText, &DialogTitle );
 					if(RetVal != EAppReturnType::Yes)
 					{
 						// Said not to overwrite (or clicked x) so bail out
@@ -1223,7 +1231,8 @@ FReply SPListEditorPanel::OnSaveClicked()
 
 		// Display Message
 		FText OverwriteMessage = LOCTEXT("PListNodeValidationFail", "Cannot save file: Not all plist entries have valid input");
-		OpenMsgDlgInt(EAppMsgType::Ok, OverwriteMessage, LOCTEXT("PListWarningCaption", "Warning") );
+		FText DialogTitle = LOCTEXT("PListWarningCaption", "Warning");
+		FMessageDialog::Open(EAppMsgType::Ok, OverwriteMessage, &DialogTitle );
 
 		return FReply::Handled();
 	}
@@ -1299,8 +1308,9 @@ FReply SPListEditorPanel::OnSaveAsClicked()
 				Arguments.Add(TEXT("Filename"), FText::FromString( InOutLastPath ));
 				FText OverwriteMessageFormatting = LOCTEXT("PListFileExistsMessageFormatting", "Overwrite existing file {Filename}?");
 				FText DialogText = FText::Format( OverwriteMessageFormatting, Arguments );
+				FText DialogTitle = LOCTEXT("PListWarningCaption", "Warning");
 
-				EAppReturnType::Type RetVal = OpenMsgDlgInt( EAppMsgType::YesNo, DialogText, LOCTEXT("PListWarningCaption", "Warning") );
+				EAppReturnType::Type RetVal = FMessageDialog::Open( EAppMsgType::YesNo, DialogText, &DialogTitle );
 				if(RetVal != EAppReturnType::Yes)
 				{
 					// Said not to overwrite (or clicked x) so bail out
@@ -1325,7 +1335,9 @@ FReply SPListEditorPanel::OnSaveAsClicked()
 
 		// Display Message
 		FText OverwriteMessage = LOCTEXT("PListNodeValidationFail", "Cannot save file: Not all plist entries have valid input");
-		OpenMsgDlgInt( EAppMsgType::Ok, OverwriteMessage, LOCTEXT("PListWarningCaption", "Warning") );
+		FText DialogTitle = LOCTEXT("PListWarningCaption", "Warning");
+
+		FMessageDialog::Open( EAppMsgType::Ok, OverwriteMessage, &DialogTitle );
 
 		return FReply::Handled();
 	}
diff --git a/Engine/Source/Editor/UnrealEd/Private/AssetSelection.cpp b/Engine/Source/Editor/UnrealEd/Private/AssetSelection.cpp
--- a/Engine/Source/Editor/UnrealEd/Private/AssetSelection.cpp
+++ b/Engine/Source/Editor/UnrealEd/Private/AssetSelection.cpp
@@ -48,7 +48,7 @@
 #include "Engine/LevelStreaming.h"
 #include "Engine/LevelBounds.h"
 #include "SourceControlHelpers.h"
-#include "Dialogs/Dialogs.h"
+#include "Misc/MessageDialog.h"
 
 
 namespace AssetSelectionUtils
@@ -454,7 +454,8 @@ namespace ActorPlacementUtils
 			FSourceControlState SCState = SourceControlHelpers::QueryFileState(FileName, true);
 			if (!InLevel->bLevelOkayForPlacementWhileCheckedIn && !(SCState.bIsCheckedOut || SCState.bIsAdded || SCState.bCanAdd || SCState.bIsUnknown))
 			{
-				if (EAppReturnType::Ok != OpenMsgDlgInt(EAppMsgType::OkCancel, NSLOCTEXT("UnrealEd","LevelNotCheckedOutMsg", "This actor will be placed in a level that is in source control but not currently checked out. Continue?"), NSLOCTEXT("UnrealEd", "LevelCheckout_Title", "Level Checkout Warning")))
+				FText Title = NSLOCTEXT("UnrealEd", "LevelCheckout_Title", "Level Checkout Warning");
+				if (EAppReturnType::Ok != FMessageDialog::Open(EAppMsgType::OkCancel, NSLOCTEXT("UnrealEd","LevelNotCheckedOutMsg", "This actor will be placed in a level that is in source control but not currently checked out. Continue?"), &Title))
 				{
 					return false;
 				}
@@ -498,7 +499,8 @@ namespace ActorPlacementUtils
 				FTransform ActorTransform = InActorTransforms[ActorTransformIndex];
 				if (!CurrentLevelBounds.IsInsideOrOn(ActorTransform.GetLocation()))
 				{
-					if (EAppReturnType::Ok != OpenMsgDlgInt(EAppMsgType::OkCancel, NSLOCTEXT("UnrealEd", "LevelBoundsMsg", "The actor will be placed outside the bounds of the current level. Continue?"), NSLOCTEXT("UnrealEd", "ActorPlacement_Title", "Actor Placement Warning")))
+					FText Title = NSLOCTEXT("UnrealEd", "ActorPlacement_Title", "Actor Placement Warning");
+					if (EAppReturnType::Ok != FMessageDialog::Open(EAppMsgType::OkCancel, NSLOCTEXT("UnrealEd", "LevelBoundsMsg", "The actor will be placed outside the bounds of the current level. Continue?"), &Title))
 					{
 						return false;
 					}
diff --git a/Engine/Source/Editor/UnrealEd/Private/Dialogs/Dialogs.cpp b/Engine/Source/Editor/UnrealEd/Private/Dialogs/Dialogs.cpp
--- a/Engine/Source/Editor/UnrealEd/Private/Dialogs/Dialogs.cpp
+++ b/Engine/Source/Editor/UnrealEd/Private/Dialogs/Dialogs.cpp
@@ -2,6 +2,7 @@
 
 
 #include "Dialogs/Dialogs.h"
+#include "Dialogs/DialogsPrivate.h"
 #include "Misc/App.h"
 #include "Misc/AssertionMacros.h"
 #include "Misc/MessageDialog.h"
@@ -296,7 +297,29 @@ void CreateMsgDlgWindow(TSharedPtr<SWindow>& OutWindow, TSharedPtr<SChoiceDialog
 	OutWindow->SetContent(OutDialog.ToSharedRef());
 }
 
-EAppReturnType::Type OpenMsgDlgInt(EAppMsgType::Type InMessageType, const FText& InMessage, const FText& InTitle)
+EAppReturnType::Type OpenMessageDialog_Internal(EAppMsgType::Type InMessageType, EAppReturnType::Type InDefaultValue, const FText& InMessage, const FText& InTitle)
+{
+	EAppReturnType::Type Response = InDefaultValue;
+	if (FApp::IsUnattended() == true || GIsRunningUnattendedScript)
+	{
+		UE_LOG(LogDialogs, Log, TEXT("Message Dialog was triggered in unattended script mode without a default value. %d will be used."), (int32)InDefaultValue);
+	}
+	else
+	{
+		TSharedPtr<SWindow> MsgWindow = NULL;
+		TSharedPtr<SChoiceDialog> MsgDialog = NULL;
+
+		CreateMsgDlgWindow(MsgWindow, MsgDialog, InMessageType, InMessage, InTitle);
+
+		GEditor->EditorAddModalWindow(MsgWindow.ToSharedRef());
+
+		Response = MsgDialog->GetResponse();
+	}
+
+	return Response;
+}
+
+EAppReturnType::Type OpenMessageDialog_Internal(EAppMsgType::Type InMessageType, const FText& InMessage, const FText& InTitle)
 {
 	EAppReturnType::Type DefaultValue = EAppReturnType::Yes;
 	switch (InMessageType)
@@ -338,30 +361,17 @@ EAppReturnType::Type OpenMsgDlgInt(EAppMsgType::Type InMessageType, const FText&
 		}
 	}
 
-	return OpenMsgDlgInt(InMessageType, DefaultValue, InMessage, InTitle);
+	return OpenMessageDialog_Internal(InMessageType, DefaultValue, InMessage, InTitle);
 }
 
-EAppReturnType::Type OpenMsgDlgInt(EAppMsgType::Type InMessageType, EAppReturnType::Type InDefaultValue, const FText& InMessage, const FText& InTitle)
+EAppReturnType::Type OpenMsgDlgInt(EAppMsgType::Type InMessageType, const FText& InMessage, const FText& InTitle)
 {
-	EAppReturnType::Type Response = InDefaultValue;
-	if (FApp::IsUnattended() == true || GIsRunningUnattendedScript)
-	{
-		UE_LOG(LogDialogs, Log, TEXT("Message Dialog was triggered in unattended script mode without a default value. %d will be used."), (int32)InDefaultValue);
-	}
-	else
-	{
-		TSharedPtr<SWindow> MsgWindow = NULL;
-		TSharedPtr<SChoiceDialog> MsgDialog = NULL;
-
-		CreateMsgDlgWindow(MsgWindow, MsgDialog, InMessageType, InMessage, InTitle);
-
-		GEditor->EditorAddModalWindow(MsgWindow.ToSharedRef());
-
-		Response = MsgDialog->GetResponse();
-	}
+	return OpenMessageDialog_Internal(InMessageType, InMessage, InTitle);
+}
 
-	UE_LOG(LogDialogs, Log, TEXT("Message dialog closed, result: %s, title: %s, text: %s"), LexToString(Response), *InTitle.ToString(), *InMessage.ToString());
-	return Response;
+EAppReturnType::Type OpenMsgDlgInt(EAppMsgType::Type InMessageType, EAppReturnType::Type InDefaultValue, const FText& InMessage, const FText& InTitle)
+{
+	return OpenMessageDialog_Internal(InMessageType, InDefaultValue, InMessage, InTitle);
 }
 
 TSharedRef<SWindow> OpenMsgDlgInt_NonModal(EAppMsgType::Type InMessageType, const FText& InMessage, const FText& InTitle,
diff --git a/Engine/Source/Editor/UnrealEd/Private/Dialogs/DialogsPrivate.h b/Engine/Source/Editor/UnrealEd/Private/Dialogs/DialogsPrivate.h
new file mode 100644
--- /dev/null
+++ b/Engine/Source/Editor/UnrealEd/Private/Dialogs/DialogsPrivate.h
@@ -0,0 +1,16 @@
+// Copyright 1998-2019 Epic Games, Inc. All Rights Reserved.
+
+#pragma once
+
+#include "CoreMinimal.h"
+
+/**
+ * Opens a modal/blocking message box dialog (with an additional 'copy message text' button), and returns the result immediately
+ * For internal use only, use FMessageDialog::Open() instead.
+ *
+ * @param InMessageType		The type of message box to display (e.g. 'ok', or 'yes'/'no' etc.)
+ * @param InMessage			The message to display in the message box
+ * @param InTitle			The title to display for the message box
+ * @return					Returns the result of the user input
+ */
+EAppReturnType::Type OpenMessageDialog_Internal(EAppMsgType::Type InMessageType, const FText& InMessage, const FText& InTitle);
diff --git a/Engine/Source/Editor/UnrealEd/Private/Dialogs/SDeleteAssetsDialog.cpp b/Engine/Source/Editor/UnrealEd/Private/Dialogs/SDeleteAssetsDialog.cpp
--- a/Engine/Source/Editor/UnrealEd/Private/Dialogs/SDeleteAssetsDialog.cpp
+++ b/Engine/Source/Editor/UnrealEd/Private/Dialogs/SDeleteAssetsDialog.cpp
@@ -14,7 +14,7 @@
 #include "Settings/EditorLoadingSavingSettings.h"
 #include "EditorDirectories.h"
 #include "FileHelpers.h"
-#include "Dialogs/Dialogs.h"
+#include "Misc/MessageDialog.h"
 #include "IContentBrowserSingleton.h"
 #include "ContentBrowserModule.h"
 #include "Editor.h"
@@ -740,7 +740,7 @@ FReply SDeleteAssetsDialog::ReplaceReferences()
 	FText Message = FText::Format( LOCTEXT( "ReplaceMessage", "This will replace any reference to the pending deleted assets with {0}; and then delete them.\n\nAre you sure?" ), FText::FromName( ConsolidationAsset.AssetName ) );
 	FText Title = LOCTEXT( "ReplaceTitle", "Replace References?" );
 
-	if ( EAppReturnType::Ok == OpenMsgDlgInt( EAppMsgType::OkCancel, Message, Title ) )
+	if ( EAppReturnType::Ok == FMessageDialog::Open( EAppMsgType::OkCancel, Message, &Title ) )
 	{
 		ParentWindow.Get()->RequestDestroyWindow();
 		DeleteRelevantSourceContent();
diff --git a/Engine/Source/Editor/UnrealEd/Private/Editor.cpp b/Engine/Source/Editor/UnrealEd/Private/Editor.cpp
--- a/Engine/Source/Editor/UnrealEd/Private/Editor.cpp
+++ b/Engine/Source/Editor/UnrealEd/Private/Editor.cpp
@@ -22,7 +22,7 @@
 #include "Misc/ScopedSlowTask.h"
 #include "UObject/UObjectIterator.h"
 #include "EngineUtils.h"
-#include "Dialogs/Dialogs.h"
+#include "Misc/MessageDialog.h"
 
 // needed for the RemotePropagator
 
@@ -532,8 +532,9 @@ void FReimportManager::ValidateAllSourceFileAndReimport(TArray<UObject*> &ToImpo
 			}
 			Arguments.Add(TEXT("AssetToFileList"), FText::FromString(AssetToFileListString));
 			FText DialogText = FText::Format(LOCTEXT("ReimportMissingFileChoiceDialogMessage", "There is {MissingNumber} assets with missing source file path. Do you want to specify a new source file path for each asset?\n \"No\" will skip the reimport of all asset with a missing source file path.\n \"Cancel\" will cancel the whole reimport.\n{AssetToFileList}"), Arguments);
+			const FText Title = LOCTEXT("ReimportMissingFileChoiceDialogMessageTitle", "Reimport missing files");
 
-			UserChoice = OpenMsgDlgInt(EAppMsgType::YesNoCancel, DialogText, LOCTEXT("ReimportMissingFileChoiceDialogMessageTitle", "Reimport missing files"));
+			UserChoice = FMessageDialog::Open(EAppMsgType::YesNoCancel, DialogText, &Title);
 		}
 
 		//Ask missing file locations
diff --git a/Engine/Source/Editor/UnrealEd/Private/EditorEngine.cpp b/Engine/Source/Editor/UnrealEd/Private/EditorEngine.cpp
--- a/Engine/Source/Editor/UnrealEd/Private/EditorEngine.cpp
+++ b/Engine/Source/Editor/UnrealEd/Private/EditorEngine.cpp
@@ -75,6 +75,7 @@
 #include "FileHelpers.h"
 #include "EditorModeInterpolation.h"
 #include "Dialogs/Dialogs.h"
+#include "Dialogs/DialogsPrivate.h"
 #include "UnrealEdGlobals.h"
 #include "Matinee/MatineeActor.h"
 #include "InteractiveFoliageActor.h"
@@ -4692,7 +4693,7 @@ EAppReturnType::Type UEditorEngine::OnModalMessageDialog(EAppMsgType::Type InMes
 {
 	if( IsInGameThread() && FSlateApplication::IsInitialized() && FSlateApplication::Get().CanAddModalWindow() )
 	{
-		return OpenMsgDlgInt(InMessage, InText, InTitle);
+		return OpenMessageDialog_Internal(InMessage, InText, InTitle);
 	}
 	else
 	{
diff --git a/Engine/Source/Editor/UnrealEd/Private/FileHelpers.cpp b/Engine/Source/Editor/UnrealEd/Private/FileHelpers.cpp
--- a/Engine/Source/Editor/UnrealEd/Private/FileHelpers.cpp
+++ b/Engine/Source/Editor/UnrealEd/Private/FileHelpers.cpp
@@ -290,7 +290,9 @@ static bool ConfirmPackageBranchCheckOutStatus(const TArray<UPackage*>& Packages
 			const FText Message = SourceControlState->IsModifiedInOtherBranch() ? FText::Format(LOCTEXT("WarningModifiedOtherBranch", "WARNING: Package {3} modified in {0} CL {1}\n\n{2}\n\nCheck out packages anyway?"), FText::FromString(HeadBranch), FText::AsNumber(HeadCL, &NoCommas), InfoText, PackageNameText)
 				: FText::Format(LOCTEXT("WarningCheckedOutOtherBranch", "WARNING: Package {2} checked out in {0}\n\n{1}\n\nCheck out packages anyway?"), FText::FromString(SourceControlState->GetOtherUserBranchCheckedOuts()), InfoText, PackageNameText);
 
-			return OpenMsgDlgInt(EAppMsgType::YesNo, Message, SourceControlState->IsModifiedInOtherBranch() ? FText::FromString("Package Branch Modifications") : FText::FromString("Package Branch Checkouts")) == EAppReturnType::Yes;
+			const FText Title = SourceControlState->IsModifiedInOtherBranch() ? FText::FromString("Package Branch Modifications") : FText::FromString("Package Branch Checkouts");
+
+			return FMessageDialog::Open(EAppMsgType::YesNo, Message, &Title) == EAppReturnType::Yes;
 		}
 	}
 
@@ -1682,8 +1684,8 @@ bool FEditorFileUtils::PromptToCheckoutPackages(bool bCheckDirty, const TArray<U
 					FText MessageFormatting = NSLOCTEXT("FileHelper", "FailedMakingWritableDlgMessageFormatting", "The following assets could not be made writable:{Packages}");
 					FText Message = FText::Format( MessageFormatting, Arguments );
 
-					FText Tile = NSLOCTEXT("FileHelper", "FailedMakingWritableDlg_Title", "Unable to make assets writable");
-					FMessageDialog::Open(EAppMsgType::Ok, Message, &Tile);
+					FText Title = NSLOCTEXT("FileHelper", "FailedMakingWritableDlg_Title", "Unable to make assets writable");
+					FMessageDialog::Open(EAppMsgType::Ok, Message, &Title);
 				}
 
 				bPerformedOperation = true;
@@ -1825,8 +1827,8 @@ ECommandResult::Type FEditorFileUtils::CheckoutPackages(const TArray<UPackage*>&
 		FText MessageFormat = NSLOCTEXT("FileHelper", "FailedCheckoutDlgMessageFormatting", "The following assets could not be successfully checked out from source control:{Packages}");
 		FText Message = FText::Format( MessageFormat, Arguments );
 
-		FText Tile = NSLOCTEXT("FileHelper", "FailedCheckoutDlg_Title", "Unable to Check Out From Source Control!");
-		FMessageDialog::Open(EAppMsgType::Ok, Message, &Tile);
+		FText Title = NSLOCTEXT("FileHelper", "FailedCheckoutDlg_Title", "Unable to Check Out From Source Control!");
+		FMessageDialog::Open(EAppMsgType::Ok, Message, &Title);
 	}
 
 	return CheckOutResult;
diff --git a/Engine/Source/Editor/UnrealEd/Private/ObjectTools.cpp b/Engine/Source/Editor/UnrealEd/Private/ObjectTools.cpp
--- a/Engine/Source/Editor/UnrealEd/Private/ObjectTools.cpp
+++ b/Engine/Source/Editor/UnrealEd/Private/ObjectTools.cpp
@@ -660,9 +660,10 @@ namespace ObjectTools
 				Arguments.Add(TEXT("Objects"), FText::FromString( RootSetObjectNames ));
 				FText MessageFormatting = NSLOCTEXT("ObjectTools", "ConsolidateAssetsRootSetDlgMsgFormatting", "The assets below were in the root set and we must remove that flag in order to proceed.  Being in the root set means that this was loaded at startup and is meant to remain in memory during gameplay.  For most assets this should be fine.  If, for some reason, there is an error, you will be notified.  Would you like to remove this flag?\n\n{Objects}");
 				FText Message = FText::Format( MessageFormatting, Arguments );
+				FText Title = NSLOCTEXT("ObjectTools", "ConsolidateAssetsRootSetDlg_Title", "Failed to Consolidate Assets");
 
 				// Prompt the user to see if they'd like to remove the root set flag from the assets and attempt to replace them
-				EAppReturnType::Type UserResponse = OpenMsgDlgInt( EAppMsgType::YesNo, EAppReturnType::No, Message, NSLOCTEXT("ObjectTools", "ConsolidateAssetsRootSetDlg_Title", "Failed to Consolidate Assets") );
+				EAppReturnType::Type UserResponse = FMessageDialog::Open( EAppMsgType::YesNo, EAppReturnType::No, Message, &Title );
 
 				// The user elected to not remove the root set flag, so cancel the replacement
 				if (UserResponse == EAppReturnType::No )
@@ -1272,8 +1273,9 @@ namespace ObjectTools
 				Arguments.Add(TEXT("Objects"), FText::FromString( FailedObjectNames ));
 				FText MessageFormatting = NSLOCTEXT("ObjectTools", "ConsolidateAssetsFailureDlgMFormattings", "The assets below were unable to be consolidated. This is likely because they are referenced by the object to consolidate to.\n\n{Objects}");
 				FText Message = FText::Format( MessageFormatting, Arguments );
+				FText Title = NSLOCTEXT("ObjectTools", "ConsolidateAssetsFailureDlg_Title", "Failed to Consolidate Assets");
 
-				OpenMsgDlgInt( EAppMsgType::Ok, Message, NSLOCTEXT("ObjectTools", "ConsolidateAssetsFailureDlg_Title", "Failed to Consolidate Assets") );
+				FMessageDialog::Open( EAppMsgType::Ok, Message, &Title );
 			}
 
 			// Alert the user to critical object failure
@@ -1298,8 +1300,9 @@ namespace ObjectTools
 				Arguments.Add(TEXT("Packages"), FText::FromString( DirtiedPackageNames ));
 				FText MessageFormatting = NSLOCTEXT("ObjectTools", "ConsolidateAssetsCriticalFailureDlgMsgFormatting", "CRITICAL FAILURE:\nOne or more assets were partially consolidated, yet still cannot be deleted for some reason. It is highly recommended that you restart the editor without saving any of the assets or packages.\n\nAffected Assets:\n{Assets}\n\nPotentially Affected Packages:\n{Packages}");
 				FText Message = FText::Format( MessageFormatting, Arguments );
+				FText Title = NSLOCTEXT("ObjectTools", "ConsolidateAssetsCriticalFailureDlg_Title", "Critical Failure to Consolidate Assets");
 
-				OpenMsgDlgInt( EAppMsgType::Ok, Message, NSLOCTEXT("ObjectTools", "ConsolidateAssetsCriticalFailureDlg_Title", "Critical Failure to Consolidate Assets") );
+				FMessageDialog::Open( EAppMsgType::Ok, Message, &Title );
 			}
 		}
 
@@ -2145,10 +2148,11 @@ namespace ObjectTools
 
 		if (ContainsWorldInUse(ObjectsToDelete))
 		{
-			OpenMsgDlgInt(
+			FText Title = NSLOCTEXT("UnrealEd", "DeleteFailedWorldInUseTitle", "Unable to delete level");
+			FMessageDialog::Open(
 				EAppMsgType::Ok,
 				NSLOCTEXT("UnrealEd", "DeleteFailedWorldInUse", "Unable to delete level while it is open"),
-				NSLOCTEXT("UnrealEd", "DeleteFailedWorldInUseTitle", "Unable to delete level")
+				&Title
 			);
 
 			return 0;
diff --git a/Engine/Source/Editor/UnrealEd/Private/PackageRestore.cpp b/Engine/Source/Editor/UnrealEd/Private/PackageRestore.cpp
--- a/Engine/Source/Editor/UnrealEd/Private/PackageRestore.cpp
+++ b/Engine/Source/Editor/UnrealEd/Private/PackageRestore.cpp
@@ -20,7 +20,7 @@
 #include "Widgets/Input/SCheckBox.h"
 #include "EditorStyleSet.h"
 #include "Editor.h"
-#include "Dialogs/Dialogs.h"
+#include "Misc/MessageDialog.h"
 #include "PackageTools.h"
 #include "AutoSaveUtils.h"
 
@@ -656,8 +656,9 @@ FEditorFileUtils::EPromptReturnCode PackageRestore::PromptToRestorePackages(cons
 		Args.Add(TEXT("FailedPackages"), FText::FromString(FailedPackagesStr));
 
 		const FText Message = FText::Format(FText::FromString("{FailedRestoreMessage}:\n{FailedPackages}"), Args);
+		const FText Title = LOCTEXT("FailedRestoreDlgTitle", "Failed to restore packages!");
 
-		OpenMsgDlgInt(EAppMsgType::Ok, Message, LOCTEXT("FailedRestoreDlgTitle", "Failed to restore packages!"));
+		FMessageDialog::Open(EAppMsgType::Ok, Message, &Title);
 
 		return FEditorFileUtils::PR_Failure;
 	}
diff --git a/Engine/Source/Editor/UnrealEd/Public/Dialogs/Dialogs.h b/Engine/Source/Editor/UnrealEd/Public/Dialogs/Dialogs.h
--- a/Engine/Source/Editor/UnrealEd/Public/Dialogs/Dialogs.h
+++ b/Engine/Source/Editor/UnrealEd/Public/Dialogs/Dialogs.h
@@ -13,16 +13,19 @@ class SModalDialogWithCheckbox;
 
 /**
  * Opens a modal/blocking message box dialog (with an additional 'copy message text' button), and returns the result immediately
+ * Internal use only. Call FMessageDialog::Open instead.
  *
  * @param InMessageType		The type of message box to display (e.g. 'ok', or 'yes'/'no' etc.)
  * @param InMessage			The message to display in the message box
  * @param InTitle			The title to display for the message box
  * @return					Returns the result of the user input
  */
+UE_DEPRECATED(4.25, "OpenMsgDlgInt is deprecated, use FMessageDialog::Open instead.")
 EAppReturnType::Type UNREALED_API OpenMsgDlgInt(EAppMsgType::Type InMessageType, const FText& InMessage, const FText& InTitle);
 
 /**
  * Opens a modal/blocking message box dialog (with an additional 'copy message text' button), and returns the result immediately
+ * Internal use only. Call FMessageDialog::Open instead.
  *
  * @param InMessageType		The type of message box to display (e.g. 'ok', or 'yes'/'no' etc.)
  * @param InDefaultValue	If the application is Unattended, the function will log and return DefaultValue
@@ -30,6 +33,7 @@ EAppReturnType::Type UNREALED_API OpenMsgDlgInt(EAppMsgType::Type InMessageType,
  * @param InTitle			The title to display for the message box
  * @return					Returns the result of the user input
 */
+UE_DEPRECATED(4.25, "OpenMsgDlgInt is deprecated, use FMessageDialog::Open instead.")
 EAppReturnType::Type UNREALED_API OpenMsgDlgInt(EAppMsgType::Type InMessageType, EAppReturnType::Type InDefaultValue, const FText& InMessage, const FText& InTitle);
 
 DECLARE_DELEGATE_TwoParams(FOnMsgDlgResult, const TSharedRef<SWindow>&, EAppReturnType::Type);
diff --git a/Engine/Source/Runtime/Core/Private/Misc/MessageDialog.cpp b/Engine/Source/Runtime/Core/Private/Misc/MessageDialog.cpp
--- a/Engine/Source/Runtime/Core/Private/Misc/MessageDialog.cpp
+++ b/Engine/Source/Runtime/Core/Private/Misc/MessageDialog.cpp
@@ -114,13 +114,11 @@ EAppReturnType::Type FMessageDialog::Open(EAppMsgType::Type MessageType, EAppRet
 {
 	EAppReturnType::Type Result = DefaultValue;
 	const FText Title = OptTitle ? *OptTitle : GetDefaultMessageTitle();
-	bool bAlreadyLogged = false;
 
 	if (!FApp::IsUnattended() && !GIsRunningUnattendedScript)
 	{
 		if ( GIsEditor && !IsRunningCommandlet() && FCoreDelegates::ModalErrorMessage.IsBound() )
 		{
-			bAlreadyLogged = true;
 			Result = FCoreDelegates::ModalErrorMessage.Execute( MessageType, Message, Title );
 		}
 		else
@@ -129,7 +127,7 @@ EAppReturnType::Type FMessageDialog::Open(EAppMsgType::Type MessageType, EAppRet
 		}
 	}
 
-	if (GWarn && !bAlreadyLogged)
+	if (GWarn)
 	{
 		GWarn->Logf(TEXT("Message dialog closed, result: %s, title: %s, text: %s"), LexToString(Result), *Title.ToString(), *Message.ToString());
 	}
